Adding ZenLaunch utility program to allow arbitrary programs to launch ZenServer while using the appropriate launch configuration.  Requires the use of one or more SponsorProcessIDs to ensure that the limited lifetime mode can always be enforced if the configuration demands it.

diff --git a/Engine/Source/Developer/Zen/Private/ZenServerInterface.cpp b/Engine/Source/Developer/Zen/Private/ZenServerInterface.cpp
--- a/Engine/Source/Developer/Zen/Private/ZenServerInterface.cpp
+++ b/Engine/Source/Developer/Zen/Private/ZenServerInterface.cpp
@@ -1720,6 +1720,26 @@ FZenServiceInstance::TryRecovery()
 	return bLastRecoveryResult;
 }
 
+bool
+FZenServiceInstance::AddSponsorProcessIDs(TArrayView<uint32> SponsorProcessIDs)
+{
+	ZenServerState State(/*bReadOnly*/ false);
+	ZenServerState::ZenServerEntry* Entry = State.LookupByEffectiveListenPort(Port);
+	if (Entry)
+	{
+		bool bAllAdded = true;
+		for (uint32 SponsorProcessID : SponsorProcessIDs)
+		{
+			if (!Entry->AddSponsorProcess(SponsorProcessID))
+			{
+				bAllAdded = false;
+			}
+		}
+		return bAllAdded;
+	}
+	return false;
+}
+
 uint16
 FZenServiceInstance::GetAutoLaunchedPort()
 {
diff --git a/Engine/Source/Developer/Zen/Public/Experimental/ZenServerInterface.h b/Engine/Source/Developer/Zen/Public/Experimental/ZenServerInterface.h
--- a/Engine/Source/Developer/Zen/Public/Experimental/ZenServerInterface.h
+++ b/Engine/Source/Developer/Zen/Public/Experimental/ZenServerInterface.h
@@ -193,6 +193,8 @@ public:
 	UE_API bool GetGCStatus(FGCStatus& Status);
 	UE_API bool RequestGC(const bool* OverrideCollectSmallObjects = nullptr, const uint32* OverrideMaxCacheDuration = nullptr);
 
+	UE_API bool AddSponsorProcessIDs(TArrayView<uint32> SponsorProcessIDs);
+
 	static UE_API uint16 GetAutoLaunchedPort();
 
 private:
diff --git a/Engine/Source/Programs/ZenLaunch/Private/PlatformWorkarounds.cpp b/Engine/Source/Programs/ZenLaunch/Private/PlatformWorkarounds.cpp
new file mode 100644
--- /dev/null
+++ b/Engine/Source/Programs/ZenLaunch/Private/PlatformWorkarounds.cpp
@@ -0,0 +1,7 @@
+// Copyright Epic Games, Inc. All Rights Reserved.
+
+#include "CoreMinimal.h"
+
+// Hack for missing Launch module dependency that causes link error on a certain platform
+FString GFileRootDirectory;
+FString GSandboxName;
\ No newline at end of file
diff --git a/Engine/Source/Programs/ZenLaunch/Private/ZenLaunch.cpp b/Engine/Source/Programs/ZenLaunch/Private/ZenLaunch.cpp
new file mode 100644
--- /dev/null
+++ b/Engine/Source/Programs/ZenLaunch/Private/ZenLaunch.cpp
@@ -0,0 +1,98 @@
+// Copyright Epic Games, Inc. All Rights Reserved.
+
+#include "ZenLaunch.h"
+
+#include "Containers/UnrealString.h"
+#include "Experimental/ZenServerInterface.h"
+#include "HAL/PlatformProcess.h"
+#include "Misc/CommandLine.h"
+#include "Misc/Parse.h"
+#include "String/LexFromString.h"
+
+#include "RequiredProgramMainCPPInclude.h"
+
+DEFINE_LOG_CATEGORY_STATIC(LogZenLaunch, Log, All);
+
+IMPLEMENT_APPLICATION(ZenLaunch, "ZenLaunch");
+
+INT32_MAIN_INT32_ARGC_TCHAR_ARGV()
+{
+	FTaskTagScope Scope(ETaskTag::EGameThread);
+	ON_SCOPE_EXIT
+	{ 
+		LLM(FLowLevelMemTracker::Get().UpdateStatsPerFrame());
+		RequestEngineExit(TEXT("Exiting"));
+		FEngineLoop::AppPreExit();
+		FModuleManager::Get().UnloadModulesAtShutdown();
+		FEngineLoop::AppExit();
+	};
+
+	if (int32 Ret = GEngineLoop.PreInit(ArgC, ArgV))
+	{
+		return Ret;
+	}
+
+	const TCHAR* CommandLine = FCommandLine::Get();
+	TArray<uint32> SponsorProcessIDs;
+	for (FString Token; FParse::Token(CommandLine, Token, /*UseEscape*/ false);)
+	{
+		TArray<FString> SponsorProcessIDStrings;
+		Token.ReplaceInline(TEXT("\""), TEXT(""));
+		const auto GetSwitchValues = [Token = FStringView(Token)](FStringView Match, TArray<FString>& OutValues)
+		{
+			if (Token.StartsWith(Match))
+			{
+				OutValues.Emplace(Token.RightChop(Match.Len()));
+			}
+		};
+
+		GetSwitchValues(TEXT("-SponsorProcessID="), SponsorProcessIDStrings);
+
+		for (const FString& SponsorProcessIDString : SponsorProcessIDStrings)
+		{
+			uint32 SponsorProcessID = 0;
+			LexFromString(SponsorProcessID, SponsorProcessIDString);
+			if (SponsorProcessID == 0)
+			{
+				UE_LOG(LogZenLaunch, Warning, TEXT("Skipping invalid sponsor process ID: %s"), *SponsorProcessIDString);
+				continue;
+			}
+
+			FProcHandle SponsorProcess = FPlatformProcess::OpenProcess(SponsorProcessID);
+			ON_SCOPE_EXIT
+			{ 
+				FPlatformProcess::CloseProc(SponsorProcess);
+			};
+
+			if (!SponsorProcess.IsValid() || !FPlatformProcess::IsProcRunning(SponsorProcess))
+			{
+				UE_LOG(LogZenLaunch, Warning, TEXT("Skipping sponsor process ID because the process is not accessible and running: %s"), *SponsorProcessIDString);
+				continue;
+			}
+
+			SponsorProcessIDs.Add(SponsorProcessID);
+		}
+	}
+
+	if (SponsorProcessIDs.IsEmpty())
+	{
+		UE_LOG(LogZenLaunch, Error, TEXT("No valid sponsor process IDs supplied on the commandline. "
+			"Please supply process IDs that require ZenServer to launch and stay running via the -SponsorProcessID=X commandline argument."));
+		return 1;
+	}
+
+	UE::Zen::FZenServiceInstance& ZenServiceInstance = UE::Zen::GetDefaultServiceInstance();
+	const UE::Zen::FServiceSettings& ZenServiceSettings = ZenServiceInstance.GetServiceSettings();
+
+	if (ZenServiceSettings.IsAutoLaunch() &&
+		ZenServiceInstance.GetServiceSettings().SettingsVariant.Get<UE::Zen::FServiceAutoLaunchSettings>().bLimitProcessLifetime)
+	{
+		if (!ZenServiceInstance.AddSponsorProcessIDs(SponsorProcessIDs))
+		{
+			UE_LOG(LogZenLaunch, Error, TEXT("Failed to add sponsor process IDs to launched ZenServer."));
+			return 1;
+		}
+	}
+
+	return 0;
+}
diff --git a/Engine/Source/Programs/ZenLaunch/Private/ZenLaunch.h b/Engine/Source/Programs/ZenLaunch/Private/ZenLaunch.h
new file mode 100644
--- /dev/null
+++ b/Engine/Source/Programs/ZenLaunch/Private/ZenLaunch.h
@@ -0,0 +1,6 @@
+// Copyright Epic Games, Inc. All Rights Reserved.
+
+#pragma once
+
+#include "CoreMinimal.h"
+
diff --git a/Engine/Source/Programs/ZenLaunch/ZenLaunch.Build.cs b/Engine/Source/Programs/ZenLaunch/ZenLaunch.Build.cs
new file mode 100644
--- /dev/null
+++ b/Engine/Source/Programs/ZenLaunch/ZenLaunch.Build.cs
@@ -0,0 +1,15 @@
+// Copyright Epic Games, Inc. All Rights Reserved.
+
+using UnrealBuildTool;
+
+public class ZenLaunch : ModuleRules
+{
+	public ZenLaunch(ReadOnlyTargetRules Target) : base(Target)
+	{
+		PublicIncludePathModuleNames.Add("Launch");
+		PrivateDependencyModuleNames.Add("Core");
+		PrivateDependencyModuleNames.Add("Projects");
+		PrivateDependencyModuleNames.Add("Json");
+		PrivateDependencyModuleNames.Add("Zen");
+	}
+}
diff --git a/Engine/Source/Programs/ZenLaunch/ZenLaunch.Target.cs b/Engine/Source/Programs/ZenLaunch/ZenLaunch.Target.cs
new file mode 100644
--- /dev/null
+++ b/Engine/Source/Programs/ZenLaunch/ZenLaunch.Target.cs
@@ -0,0 +1,35 @@
+// Copyright Epic Games, Inc. All Rights Reserved.
+
+using UnrealBuildTool;
+using System.Collections.Generic;
+
+[SupportedPlatforms(UnrealPlatformClass.All)]
+public class ZenLaunchTarget : TargetRules
+{
+	public ZenLaunchTarget(TargetInfo Target) : base(Target)
+	{
+		Type = TargetType.Program;
+		IncludeOrderVersion = EngineIncludeOrderVersion.Latest;
+		LinkType = TargetLinkType.Monolithic;
+		LaunchModuleName = "ZenLaunch";
+
+		// Lean and mean
+		bBuildDeveloperTools = false;
+		bBuildWithEditorOnlyData = false;
+
+		// Currently this app is not linking against the engine, so we'll compile out references from Core to the rest of the engine
+		bCompileAgainstEngine = false;
+		bCompileAgainstCoreUObject = false;
+		bCompileAgainstApplicationCore = false;
+		bCompileICU = false;
+
+		// to build with automation tests:
+		// bForceCompileDevelopmentAutomationTests = true;
+
+		// to enable tracing:
+		// bEnableTrace = true;
+
+		// This app is a console application, not a Windows app (sets entry point to main(), instead of WinMain())
+		bIsBuildingConsoleApplication = true;
+	}
+}
