Modify ZenLaunch target config to enable editor-only data as a compile-time option.  This is required for it to accept a project argument and load project-specific configuration.

diff --git a/Engine/Source/Programs/ZenLaunch/ZenLaunch.Target.cs b/Engine/Source/Programs/ZenLaunch/ZenLaunch.Target.cs
--- a/Engine/Source/Programs/ZenLaunch/ZenLaunch.Target.cs
+++ b/Engine/Source/Programs/ZenLaunch/ZenLaunch.Target.cs
@@ -15,7 +15,9 @@ public class ZenLaunchTarget : TargetRules
 
 		// Lean and mean
 		bBuildDeveloperTools = false;
-		bBuildWithEditorOnlyData = false;
+
+		// Building with editor-only data is required for accepting a project argument and loading project-specific configuration
+		bBuildWithEditorOnlyData = true;
 
 		// Currently this app is not linking against the engine, so we'll compile out references from Core to the rest of the engine
 		bCompileAgainstEngine = false;
