It is important that ZenLaunch use the appropriate project-specific configuration when launching ZenServer.  However, programs do not load project config even if a project file is specified as a commandline argument.  To address this, some programs that need project-specific configuration will use the "OverrideProjectDir".  The implementation in this change is taken from UnrealVirtualizationTool.cpp which is a program that has the same need.

diff --git a/Engine/Source/Programs/ZenLaunch/Private/ZenLaunch.cpp b/Engine/Source/Programs/ZenLaunch/Private/ZenLaunch.cpp
--- a/Engine/Source/Programs/ZenLaunch/Private/ZenLaunch.cpp
+++ b/Engine/Source/Programs/ZenLaunch/Private/ZenLaunch.cpp
@@ -27,6 +27,33 @@ INT32_MAIN_INT32_ARGC_TCHAR_ARGV()
 		FEngineLoop::AppExit();
 	};
 
+	// Although standalone tools can set the project path via the cmdline this will not change the ProjectDir being
+	// used as standalone tools have a beespoke path in FGenericPlatformMisc::ProjectDir. We can work around this
+	// for now by doing our own parsing then using the project dir override feature.
+	// This is a band aid while we consider adding better support for project files/directories with stand alone tools.
+	if (ArgC >= 2)
+	{
+		FString Cmd(ArgV[1]);
+
+		if (!Cmd.IsEmpty() && !Cmd.StartsWith(TEXT("-")) && Cmd.EndsWith(FProjectDescriptor::GetExtension()))
+		{
+			FString ProjectDir = FPaths::GetPath(Cmd);
+			ProjectDir = FFileManagerGeneric::DefaultConvertToRelativePath(*ProjectDir);
+
+			// The path should end with a trailing slash (see FGenericPlatformMisc::ProjectDir) so we use
+			// NormalizeFilename not NormalizeDirectoryName as the latter will remove trailing slashes. We
+			// also need to add one if it is missing.
+			// We probably should move this path fixup code to 'FPlatformMisc::SetOverrideProjectDir'
+			FPaths::NormalizeFilename(ProjectDir);
+			if (!ProjectDir.EndsWith(TEXT("/")))
+			{
+				ProjectDir += TEXT("/");
+			}
+
+			FPlatformMisc::SetOverrideProjectDir(ProjectDir);
+		}
+	}
+
 	if (int32 Ret = GEngineLoop.PreInit(ArgC, ArgV))
 	{
 		return Ret;
