ZenLaunch: Add support for -Sponsor=None, add shorter commandline argument -Sponsor as an alias for existing -SponsorProcessId, and allow -SponsorProcessId to specify multiple comma-delimited processids.
#rnx
#Synced-CL 52123874

diff --git a/Engine/Source/Programs/ZenLaunch/Private/ZenLaunch.cpp b/Engine/Source/Programs/ZenLaunch/Private/ZenLaunch.cpp
--- a/Engine/Source/Programs/ZenLaunch/Private/ZenLaunch.cpp
+++ b/Engine/Source/Programs/ZenLaunch/Private/ZenLaunch.cpp
@@ -8,6 +8,7 @@
 #include "Misc/CommandLine.h"
 #include "Misc/Parse.h"
 #include "String/LexFromString.h"
+#include "String/ParseTokens.h"
 #include "ProjectUtilities.h"
 
 #include "RequiredProgramMainCPPInclude.h"
@@ -36,8 +37,10 @@ INT32_MAIN_INT32_ARGC_TCHAR_ARGV()
 		return Ret;
 	}
 
+	constexpr FStringView SponsorNone(TEXTVIEW("None"));
 	const TCHAR* CommandLine = FCommandLine::Get();
 	TArray<uint32> SponsorProcessIDs;
+	bool bHasSponsorNone = false;
 	for (FString Token; FParse::Token(CommandLine, Token, /*UseEscape*/ false);)
 	{
 		TArray<FString> SponsorProcessIDStrings;
@@ -46,14 +49,26 @@ INT32_MAIN_INT32_ARGC_TCHAR_ARGV()
 		{
 			if (Token.StartsWith(Match))
 			{
-				OutValues.Emplace(Token.RightChop(Match.Len()));
+				FStringView Values = Token.RightChop(Match.Len());
+				UE::String::ParseTokens(Values, TEXTVIEW(","),
+					[&OutValues](FStringView Value)
+					{
+						OutValues.Emplace(Value);
+					},
+					UE::String::EParseTokensOptions::Trim | UE::String::EParseTokensOptions::SkipEmpty);
 			}
 		};
 
 		GetSwitchValues(TEXT("-SponsorProcessID="), SponsorProcessIDStrings);
+		GetSwitchValues(TEXT("-Sponsor="), SponsorProcessIDStrings);
 
 		for (const FString& SponsorProcessIDString : SponsorProcessIDStrings)
 		{
+			if (SponsorNone.Equals(SponsorProcessIDString, ESearchCase::IgnoreCase))
+			{
+				bHasSponsorNone = true;
+				continue;
+			}
 			uint32 SponsorProcessID = 0;
 			LexFromString(SponsorProcessID, SponsorProcessIDString);
 			if (SponsorProcessID == 0)
@@ -74,29 +89,67 @@ INT32_MAIN_INT32_ARGC_TCHAR_ARGV()
 				continue;
 			}
 
-			SponsorProcessIDs.Add(SponsorProcessID);
+			SponsorProcessIDs.AddUnique(SponsorProcessID);
 		}
 	}
 
-	if (SponsorProcessIDs.IsEmpty())
+	if (!bHasSponsorNone && SponsorProcessIDs.IsEmpty())
 	{
-		UE_LOG(LogZenLaunch, Error, TEXT("No valid sponsor process IDs supplied on the commandline. "
-			"Please supply process IDs that require ZenServer to launch and stay running via the -SponsorProcessID=X commandline argument."));
+		UE_LOG(LogZenLaunch, Error, TEXT("No valid sponsor process IDs (or -Sponsor=None) supplied on the commandline. ")
+			TEXT("\n\tPlease supply process IDs via -Sponsor=X,Y,Z. ZenServer will stay running until these processes exit.")
+			TEXT("\n\tOr pass -Sponsor=None to indicate the zenserver should run until closed via a `zen down` call."));
 		return 1;
 	}
+	if (bHasSponsorNone)
+	{
+		FPlatformMisc::SetEnvironmentVar(TEXT("UE-ZenLimitProcessLifetime"), TEXT("false"));
+	}
 
 	UE::Zen::FZenServiceInstance& ZenServiceInstance = UE::Zen::GetDefaultServiceInstance();
 	const UE::Zen::FServiceSettings& ZenServiceSettings = ZenServiceInstance.GetServiceSettings();
 
-	if (ZenServiceSettings.IsAutoLaunch() &&
-		ZenServiceInstance.GetServiceSettings().SettingsVariant.Get<UE::Zen::FServiceAutoLaunchSettings>().bLimitProcessLifetime)
+	if (!ZenServiceSettings.IsAutoLaunch())
 	{
+		UE_LOG(LogZenLaunch, Display, TEXT("Cannot set ZenServer process lifetime because the ZenServer was manually launched. ")
+			TEXT("\n\tProcess lifetime is settable only for AutoLaunched ZenServers."));
+		return 0;
+	}
+
+	if (bHasSponsorNone)
+	{
+		if (ZenServiceInstance.GetServiceSettings().SettingsVariant.Get<UE::Zen::FServiceAutoLaunchSettings>().bLimitProcessLifetime)
+		{
+			UE_LOG(LogZenLaunch, Error, TEXT("Failed to set zen to run unsponsored; unknown error. ")
+				TEXT("\n\tUE-ZenLimitProcessLifetime=false was set but upon creation bLimitProcessLifetime=true."));
+			return 1;
+		}
+
+		UE_LOG(LogZenLaunch, Display, TEXT("ZenServer is now running unsponsored; it will remain running until `zen down` is called."));
+		return 0;
+	}
+	else
+	{
+		if (!ZenServiceInstance.GetServiceSettings().SettingsVariant.Get<UE::Zen::FServiceAutoLaunchSettings>().bLimitProcessLifetime)
+		{
+			UE_LOG(LogZenLaunch, Display, TEXT("ZenServer settings specify bLimitProcessLifetime=false, so sponsors were ignored. ")
+				TEXT("\n\tZenServer is now running unsponsored; it will remain running until `zen down` is called."));
+			return 0;
+		}
+
 		if (!ZenServiceInstance.AddSponsorProcessIDs(SponsorProcessIDs))
 		{
 			UE_LOG(LogZenLaunch, Error, TEXT("Failed to add sponsor process IDs to launched ZenServer."));
 			return 1;
 		}
-	}
 
-	return 0;
+		TArray<FString> SponsorProcessIDStrs;
+		for (uint32 ID : SponsorProcessIDs)
+		{
+			SponsorProcessIDStrs.Add(LexToString(ID));
+		}
+		UE_LOG(LogZenLaunch, Display, TEXT("ZenServer is now running with sponsor ids [%s]. ")
+			TEXT("\n\tIt will remain running until those processes exit or zen down is called."),
+			*FString::Join(SponsorProcessIDStrs, TEXT(",")));
+		return 0;
+	}
 }
